NNE - Onnx External data support : add import capability for linux

Problem: Onnx model with data cannot be imported on linux as we were missing the related shared lib to parse where to load the external data files.

Solution:
* Compile NNEEditorOnnxTools shared lib on linux and enable the external data import functionality on linux.

Test:
* Win64 Editor Debug: external data tests (including reimport) + smoke tests + manual import of the model with external data use cases (one file, multi-file, with subfolders)
* Linux Editor Development: external data tests (including reimport) + smoke tests + manual import of the model with external data use cases (one file, multi-file, with subfolders)

#rnx

diff --git a/Engine/Source/Editor/NNEEditor/NNEEditor.Build.cs b/Engine/Source/Editor/NNEEditor/NNEEditor.Build.cs
--- a/Engine/Source/Editor/NNEEditor/NNEEditor.Build.cs
+++ b/Engine/Source/Editor/NNEEditor/NNEEditor.Build.cs
@@ -28,9 +28,10 @@ public class NNEEditor : ModuleRules
 			}
 		);
 
+		string SharedLibPath = Path.Combine(ModuleDirectory, "Bin", Target.Platform.ToString());
+
 		if (Target.Platform == UnrealTargetPlatform.Win64)
 		{
-			string SharedLibPath = Path.Combine(ModuleDirectory, "Bin", Target.Platform.ToString());
 			string SharedLibFileName = "NNEEditorOnnxTools.dll";
 
 			PublicDelayLoadDLLs.Add(SharedLibFileName);
@@ -40,6 +41,17 @@ public class NNEEditor : ModuleRules
 			PublicDefinitions.Add("NNEEDITORONNXTOOLS_SUPPORTED");
 			PublicDefinitions.Add("NNEEDITORONNXTOOLS_SHAREDLIB_FILENAME=" + SharedLibFileName);
 		}
+		else if (Target.Platform == UnrealTargetPlatform.Linux)
+		{
+			string SharedLibFileName = "libNNEEditorOnnxTools.so";
+
+			PublicDelayLoadDLLs.Add(SharedLibFileName);
+
+			RuntimeDependencies.Add("$(TargetOutputDir)/" + SharedLibFileName, Path.Combine(SharedLibPath, SharedLibFileName));
+
+			PublicDefinitions.Add("NNEEDITORONNXTOOLS_SUPPORTED");
+			PublicDefinitions.Add("NNEEDITORONNXTOOLS_SHAREDLIB_FILENAME=" + SharedLibFileName);
+		}
 
 		PublicDefinitions.Add("UE_NNEEDITORONNXTOOLS");
 	}
diff --git a/Engine/Source/Editor/NNEEditor/Private/NNEEditorOnnxTools.h b/Engine/Source/Editor/NNEEditor/Private/NNEEditorOnnxTools.h
--- a/Engine/Source/Editor/NNEEditor/Private/NNEEditorOnnxTools.h
+++ b/Engine/Source/Editor/NNEEditor/Private/NNEEditorOnnxTools.h
@@ -2,10 +2,18 @@
 
 #pragma once
 
-#ifdef NNEEDITORONNXTOOLS_EXPORT
-#define NNEEDITORONNXTOOLS_API __declspec(dllexport)
+#if defined(NNEEDITORONNXTOOLS_EXPORT)
+    #if defined(_MSC_VER)
+        //  Microsoft 
+        #define NNEEDITORONNXTOOLS_API __declspec(dllexport)
+    #elif defined(__GNUC__)
+        //  GCC
+        #define NNEEDITORONNXTOOLS_API __attribute__((visibility("default")))
+    #else
+	    #define NNEEDITORONNXTOOLS_API unsupported_platform
+    #endif
 #else
-#define NNEEDITORONNXTOOLS_API
+    #define NNEEDITORONNXTOOLS_API
 #endif
 
 #ifndef UE_NNEEDITORONNXTOOLS
