[NNE] ONNX external data support, add import capability on Mac.

Problem:
Onnx model with data cannot be imported on Mac as we were missing the related shared lib to parse where to load the external data files.

Solution:
Compile NNEEditorOnnxTools shared lib on Mac (Universal Binary) and enable the external data import functionality on Mac.

Test:
Mac Editor Debug: external data tests (including reimport) + smoke tests + manual import of the model with external data use cases (one file, multi-file, with subfolders)

#tests Automation tests, reimport, smoke tests and manual import

#rnx

diff --git a/Engine/Source/Editor/NNEEditor/NNEEditor.Build.cs b/Engine/Source/Editor/NNEEditor/NNEEditor.Build.cs
--- a/Engine/Source/Editor/NNEEditor/NNEEditor.Build.cs
+++ b/Engine/Source/Editor/NNEEditor/NNEEditor.Build.cs
@@ -52,6 +52,19 @@ public class NNEEditor : ModuleRules
 			PublicDefinitions.Add("NNEEDITORONNXTOOLS_SUPPORTED");
 			PublicDefinitions.Add("NNEEDITORONNXTOOLS_SHAREDLIB_FILENAME=" + SharedLibFileName);
 		}
+		else if (Target.Platform == UnrealTargetPlatform.Mac)
+		{
+			string SharedLibFileName = "libNNEEditorOnnxTools.dylib";
+
+			// Not fully supported on macOS:
+			// https://dev.epicgames.com/documentation/en-us/unreal-engine/integrating-third-party-libraries-into-unreal-engine
+			// PublicDelayLoadDLLs.Add(SharedLibFileName);
+
+			RuntimeDependencies.Add("$(TargetOutputDir)/" + SharedLibFileName, Path.Combine(SharedLibPath, SharedLibFileName));
+
+			PublicDefinitions.Add("NNEEDITORONNXTOOLS_SUPPORTED");
+			PublicDefinitions.Add("NNEEDITORONNXTOOLS_SHAREDLIB_FILENAME=" + SharedLibFileName);
+		}
 
 		PublicDefinitions.Add("UE_NNEEDITORONNXTOOLS");
 	}
