[NNE] Backout ability to import onnx with external data while linux shared lib get recompiled against the correct glibc

Problem: The current libNNEEditorOnnxTools.so is requiring glibc 3.4.29 while min spec for rocky linux 8 is 2.28 (was added with CL CL 36002629)
Solution: Temporarily remove the shared lib (and ability to import onnx file with external data) on linux while shared lib is recompiled with rocky linyx 8 glibc dependancy.

#Test: all NNE tests on linux editor64 dev and win 64 editor debug

#rnx

diff --git a/Engine/Source/Editor/NNEEditor/NNEEditor.Build.cs b/Engine/Source/Editor/NNEEditor/NNEEditor.Build.cs
--- a/Engine/Source/Editor/NNEEditor/NNEEditor.Build.cs
+++ b/Engine/Source/Editor/NNEEditor/NNEEditor.Build.cs
@@ -41,17 +41,6 @@ public class NNEEditor : ModuleRules
 			PublicDefinitions.Add("NNEEDITORONNXTOOLS_SUPPORTED");
 			PublicDefinitions.Add("NNEEDITORONNXTOOLS_SHAREDLIB_FILENAME=" + SharedLibFileName);
 		}
-		else if (Target.Platform == UnrealTargetPlatform.Linux)
-		{
-			string SharedLibFileName = "libNNEEditorOnnxTools.so";
-
-			PublicDelayLoadDLLs.Add(SharedLibFileName);
-
-			RuntimeDependencies.Add("$(TargetOutputDir)/" + SharedLibFileName, Path.Combine(SharedLibPath, SharedLibFileName));
-
-			PublicDefinitions.Add("NNEEDITORONNXTOOLS_SUPPORTED");
-			PublicDefinitions.Add("NNEEDITORONNXTOOLS_SHAREDLIB_FILENAME=" + SharedLibFileName);
-		}
 		else if (Target.Platform == UnrealTargetPlatform.Mac)
 		{
 			string SharedLibFileName = "libNNEEditorOnnxTools.dylib";
