[NNE] Reenable ability to import onnx file with external data on linux

Problem: the feature was backed out as it relyed on an shared lib witch was incompatible with linux min spec (aka Rocky 8)
Solution: Recompile the share lib using a Rocky 8 docker image and reenable the feature.

#Test: Linux editor: smoke tests + external data tests + manual import onnx files with external data on linux (3 cases from tests)

#rnx

diff --git a/Engine/Source/Editor/NNEEditor/NNEEditor.Build.cs b/Engine/Source/Editor/NNEEditor/NNEEditor.Build.cs
--- a/Engine/Source/Editor/NNEEditor/NNEEditor.Build.cs
+++ b/Engine/Source/Editor/NNEEditor/NNEEditor.Build.cs
@@ -41,6 +41,17 @@ public class NNEEditor : ModuleRules
 			PublicDefinitions.Add("NNEEDITORONNXTOOLS_SUPPORTED");
 			PublicDefinitions.Add("NNEEDITORONNXTOOLS_SHAREDLIB_FILENAME=" + SharedLibFileName);
 		}
+		else if (Target.Platform == UnrealTargetPlatform.Linux)
+		{
+			string SharedLibFileName = "libNNEEditorOnnxTools.so";
+
+			PublicDelayLoadDLLs.Add(SharedLibFileName);
+
+			RuntimeDependencies.Add("$(TargetOutputDir)/" + SharedLibFileName, Path.Combine(SharedLibPath, SharedLibFileName));
+
+			PublicDefinitions.Add("NNEEDITORONNXTOOLS_SUPPORTED");
+			PublicDefinitions.Add("NNEEDITORONNXTOOLS_SHAREDLIB_FILENAME=" + SharedLibFileName);
+		}
 		else if (Target.Platform == UnrealTargetPlatform.Mac)
 		{
 			string SharedLibFileName = "libNNEEditorOnnxTools.dylib";
