Add Windows Arm64 version of NNEEditorOnnxTools.dll

#rnx
#test importing sample file "candy9-ext-data-one-file.onnx"

diff --git a/Engine/Source/Editor/NNEEditor/NNEEditor.Build.cs b/Engine/Source/Editor/NNEEditor/NNEEditor.Build.cs
--- a/Engine/Source/Editor/NNEEditor/NNEEditor.Build.cs
+++ b/Engine/Source/Editor/NNEEditor/NNEEditor.Build.cs
@@ -32,11 +32,12 @@ public class NNEEditor : ModuleRules
 
 		if (Target.Platform == UnrealTargetPlatform.Win64)
 		{
+			string ArchPath = (Target.Architecture == UnrealArch.Arm64) ? "arm64/" : "";
 			string SharedLibFileName = "NNEEditorOnnxTools.dll";
 
 			PublicDelayLoadDLLs.Add(SharedLibFileName);
 
-			RuntimeDependencies.Add("$(TargetOutputDir)/"+ SharedLibFileName, Path.Combine(SharedLibPath, SharedLibFileName));
+			RuntimeDependencies.Add("$(TargetOutputDir)/"+ ArchPath + SharedLibFileName, Path.Combine(SharedLibPath, ArchPath + SharedLibFileName));
 
 			PublicDefinitions.Add("NNEEDITORONNXTOOLS_SUPPORTED");
 			PublicDefinitions.Add("NNEEDITORONNXTOOLS_SHAREDLIB_FILENAME=" + SharedLibFileName);
diff --git a/Engine/Source/Editor/NNEEditor/Private/NNEEditorModule.cpp b/Engine/Source/Editor/NNEEditor/Private/NNEEditorModule.cpp
--- a/Engine/Source/Editor/NNEEditor/Private/NNEEditorModule.cpp
+++ b/Engine/Source/Editor/NNEEditor/Private/NNEEditorModule.cpp
@@ -38,7 +38,11 @@ public:
 	virtual void StartupModule() override
 	{
 #ifdef NNEEDITORONNXTOOLS_SUPPORTED
+#if PLATFORM_WINDOWS && PLATFORM_CPU_ARM_FAMILY && !PLATFORM_WINDOWS_ARM64EC
+		const FString ModuleDir = FPlatformProcess::GetModulesDirectory() / TEXT("arm64");
+#else
 		const FString ModuleDir = FPlatformProcess::GetModulesDirectory();
+#endif
 		const FString NNEEditorOnnxToolsSharedLibPath = FPaths::Combine(ModuleDir, TEXT(PREPROCESSOR_TO_STRING(NNEEDITORONNXTOOLS_SHAREDLIB_FILENAME)));
 		NNEEditorOnnxToolsSharedLibHandle = UE::NNEEditor::Private::OnnxHelper::GetSharedLibHandle(NNEEditorOnnxToolsSharedLibPath);
 		UE::NNEEditor::Private::OnnxModelInspectorHelper::SetupSharedLibFunctionPointer(NNEEditorOnnxToolsSharedLibHandle);
