[NNE] Add per model runtime settings

#rnx
#tests NNEModelTest

New a runtime can implement the CreateDefaultRuntimeSettings() function to provide a settings object that gets stored alongside the model data. These settings can be changed on the asset. The runtime can get access to its settings by calling UNNEModelData::DeserializeRuntimeSettings and passing in the AdditionalFileData argument that was passed into CanCreateModelData and CreateModelData.

#Synced-CL 49878051

diff --git a/Engine/Source/Editor/NNEEditor/Private/NNEEditorModelDataActions.cpp b/Engine/Source/Editor/NNEEditor/Private/NNEEditorModelDataActions.cpp
--- a/Engine/Source/Editor/NNEEditor/Private/NNEEditorModelDataActions.cpp
+++ b/Engine/Source/Editor/NNEEditor/Private/NNEEditorModelDataActions.cpp
@@ -2,6 +2,7 @@
 
 #include "NNEEditorModelDataActions.h"
 
+#include "NNE.h"
 #include "NNEEditorModelDataEditorToolkit.h"
 #include "NNEModelData.h"
 
@@ -32,13 +33,22 @@ namespace UE::NNEEditor::Private
 
 	void FModelDataAssetTypeActions::OpenAssetEditor(const TArray<UObject*>& InObjects, TSharedPtr<class IToolkitHost> EditWithinLevelEditor)
 	{
-		for (UObject* Obj : InObjects)
+		TArray<UNNEModelData*> Models;
+		for (UObject* Object : InObjects)
 		{
-			UNNEModelData* ModelData = Cast<UNNEModelData>(Obj);
-			if (ModelData != nullptr)
+			UNNEModelData* Model = Cast<UNNEModelData>(Object);
+			if (Model)
 			{
-				MakeShared<FModelDataEditorToolkit>()->InitEditor(ModelData);
+				Models.Add(Model);
 			}
+			else
+			{
+				UE_LOG(LogNNE, Warning, TEXT("Casting asset to UNNEModelData failed"));
+			}
+		}
+		if (!Models.IsEmpty())
+		{
+			MakeShared<FModelDataEditorToolkit>()->InitEditor(Models);
 		}
 	}
 
diff --git a/Engine/Source/Editor/NNEEditor/Private/NNEEditorModelDataEditorToolkit.cpp b/Engine/Source/Editor/NNEEditor/Private/NNEEditorModelDataEditorToolkit.cpp
--- a/Engine/Source/Editor/NNEEditor/Private/NNEEditorModelDataEditorToolkit.cpp
+++ b/Engine/Source/Editor/NNEEditor/Private/NNEEditorModelDataEditorToolkit.cpp
@@ -2,17 +2,95 @@
 
 #include "NNEEditorModelDataEditorToolkit.h"
 
+#include "Kismet/KismetSystemLibrary.h"
+#include "Modules/ModuleManager.h"
 #include "NNE.h"
 #include "NNEModelData.h"
-#include "NNEEditorRuntimesWidget.h"
+#include "NNERuntime.h"
+#include "PropertyEditorModule.h"
 #include "Widgets/Docking/SDockTab.h"
+#include "Widgets/Input/SCheckBox.h"
+#include "Widgets/Layout/SExpandableArea.h"
+
 
 namespace UE::NNEEditor::Private
 {
+	class SRuntimeExpandableArea : public SExpandableArea
+	{
+
+	public:
+
+		TAttribute<bool> IsExpandable;
+
+		void Construct(const SExpandableArea::FArguments& InArgs)
+		{
+			SExpandableArea::Construct(InArgs);
+			const FSlateColor TintColor(FLinearColor(FVector4f(0, 0, 0, 0)));
+			EmptyImage = *ExpandedImage;
+			EmptyImage.TintColor = TintColor;
+		}
+
+	protected:
+
+		const FSlateBrush* OnGetCollapseImage() const override
+		{
+			if (IsExpandable.Get())
+			{
+				return bAreaCollapsed ? CollapsedImage : ExpandedImage;
+			}
+			return &EmptyImage;
+		}
+
+		void OnToggleContentVisibility() override
+		{
+			if (IsExpandable.Get())
+			{
+				SExpandableArea::OnToggleContentVisibility();
+			}
+		}
+
+	private:
+
+		FSlateBrush EmptyImage;
+
+	};
+
+	static void AppendUniqueRuntimes(TArray<FString>& InOutDest, const TArrayView<const FString>& InSource)
+	{
+		for (const FString& RuntimeName : InSource)
+		{
+			InOutDest.AddUnique(RuntimeName);
+		}
+	}
 
-	void FModelDataEditorToolkit::InitEditor(UNNEModelData* InModelData)
+	void FModelDataEditorToolkit::InitEditor(const TArray<UNNEModelData*>& InModels)
 	{
-		ModelData = InModelData;
+		checkf(InModels.Num() > 0, TEXT("FModelDataEditorToolkit require at least one item in InModels"))
+		Models = InModels;
+		AllRuntimeNames = UE::NNE::GetAllRuntimeNames();
+		//A runtime could possibly be in the target list in any of the models but not be registered to the editor
+		for (UNNEModelData* Model : Models)
+		{
+			AppendUniqueRuntimes(AllRuntimeNames, Model->GetTargetRuntimes());
+		}
+		EmptySelectionModels = NewObject<UNNEEditorEmptySelectionModels>();
+		EmptySelectionModels->AddToRoot();
+		EmptySelectionModels->SetFlags(RF_Transient | RF_Transactional);
+		FDetailsViewArgs DetailsViewArgs{};
+		DetailsViewArgs.bAllowSearch = false;
+		DetailsViewArgs.bHideSelectionTip = true;
+		DetailsViewArgs.bShowObjectLabel = true;
+		DetailsViewArgs.bShowSectionSelector = false;
+		DetailsViewArgs.bCustomNameAreaLocation = true;
+		Runtimes.Reset(AllRuntimeNames.Num());
+		for(const FString& RuntimeName : AllRuntimeNames)
+		{
+			TSharedRef<IDetailsView> View = FModuleManager::LoadModuleChecked<FPropertyEditorModule>("PropertyEditor").CreateDetailView(DetailsViewArgs);
+			View->OnFinishedChangingProperties().AddRaw(this, &FModelDataEditorToolkit::OnRuntimePropertyChanged);
+			Runtimes.Emplace(RuntimeName, View);
+		}
+		ReloadSettings();
+		PostUndoRedoHandle = FEditorDelegates::PostUndoRedo.AddRaw(this, &FModelDataEditorToolkit::ReloadSettings);
 
 		const TSharedRef<FTabManager::FLayout> Layout = FTabManager::NewLayout("NNEModelDataEditorLayout")
 			->AddArea
@@ -25,8 +103,12 @@ namespace UE::NNEEditor::Private
 					->AddTab(TEXT("NNEModelDataTab"), ETabState::OpenedTab)
 				)
 			);
-
-		FAssetEditorToolkit::InitAssetEditor(EToolkitMode::Standalone, {}, TEXT("NNEModelDataEditor"), Layout, true, true, ModelData);
+		TArray<UObject*> UObjectModels;
+		for (UNNEModelData* Model : Models)
+		{
+			UObjectModels.Add(Model);
+		}
+		FAssetEditorToolkit::InitAssetEditor(EToolkitMode::Standalone, {}, TEXT("NNEModelDataEditor"), Layout, true, true, UObjectModels);
 	}
 
 	void FModelDataEditorToolkit::RegisterTabSpawners(const TSharedRef<class FTabManager>& InTabManager)
@@ -35,15 +117,7 @@ namespace UE::NNEEditor::Private
 
 		WorkspaceMenuCategory = InTabManager->AddLocalWorkspaceMenuCategory(INVTEXT("NNE Model Data Editor"));
 
-		InTabManager->RegisterTabSpawner("NNEModelDataTab", FOnSpawnTab::CreateLambda([this](const FSpawnTabArgs&)
-			{
-				return SNew(SDockTab)
-				[
-					SNew(SRuntimesWidget)
-					.TargetRuntimes(this, &FModelDataEditorToolkit::GetTargetRuntimes)
-					.OnTargetRuntimesChanged(this, &FModelDataEditorToolkit::SetTargetRuntimes)
-				];
-			}))
+		InTabManager->RegisterTabSpawner("NNEModelDataTab", FOnSpawnTab::CreateRaw(this, &FModelDataEditorToolkit::SpawnTab))
 			.SetDisplayName(INVTEXT("NNE Model Data"))
 			.SetGroup(WorkspaceMenuCategory.ToSharedRef());
 	}
@@ -54,17 +128,354 @@ namespace UE::NNEEditor::Private
 		InTabManager->UnregisterTabSpawner("NNEModelDataTab");
 	}
 
-	TArrayView<const FString> FModelDataEditorToolkit::GetTargetRuntimes() const
+	FText FModelDataEditorToolkit::GetToolkitName() const
+	{
+		check(Models.Num() > 0);
+		if (Models.Num() == 1)
+		{
+			return GetLabelForObject(Models[0]);
+		}
+		return INVTEXT("Multiple NNE Models");
+	}
+
+	#define LOCTEXT_NAMESPACE "AssetEditorToolkit"
+	FText FModelDataEditorToolkit::GetToolkitToolTipText() const
+	{
+		check(Models.Num() > 0);
+		if (Models.Num() == 1)
+		{
+			return GetToolTipTextForObject(Models[0]);
+		}
+		FString ToolTipString;
+		ToolTipString += LOCTEXT("ToolTipAssetLabel", "Asset").ToString();
+		ToolTipString += TEXT(":\n");
+		
+		for (const UNNEModelData* Model : Models)
+		{
+			ToolTipString += TEXT("- ");
+			ToolTipString += Model->GetName();
+			ToolTipString += TEXT("\n");
+		}
+		ToolTipString.TrimEndInline();
+		return FText::FromString(ToolTipString);
+	}
+	#undef LOCTEXT_NAMESPACE
+
+	void FModelDataEditorToolkit::OnClose()
+	{
+		EmptySelectionModels->RemoveFromRoot();
+		ClearSettings();
+		FEditorDelegates::PostUndoRedo.Remove(PostUndoRedoHandle);
+	}
+
+	TSharedRef<SDockTab> FModelDataEditorToolkit::SpawnTab(const FSpawnTabArgs&)
+	{
+		TSharedPtr<SVerticalBox> VerticalBox;
+		TSharedRef<SDockTab> Result = SNew(SDockTab)
+			[
+				SAssignNew(VerticalBox, SVerticalBox)
+					+
+					SVerticalBox::Slot()
+					.AutoHeight()
+					.Padding(4.0, 20.0, 4.0, 4.0)
+					[
+						SNew(SHorizontalBox)
+							+
+							SHorizontalBox::Slot()
+							.AutoWidth()
+							.Padding(0, 0, 15, 0)
+							[
+								SNew(STextBlock).Text(FText::FromString("Runtimes"))
+									.Font(FAppStyle::GetFontStyle("PropertyWindow.BoldFont"))
+							]
+							+
+							SHorizontalBox::Slot()
+							.VAlign(VAlign_Center)
+							
+							[
+								SNew(STextBlock).Text(FText::FromString("No runtime selected. Please select at least one runtime for each model."))
+									.ColorAndOpacity(FSlateColor(EStyleColor::AccentRed))
+									.Visibility(this, &FModelDataEditorToolkit::GetErrorMessageVisibility)
+									.Font(FAppStyle::GetFontStyle("PropertyWindow.BoldFont"))
+							]
+					]
+					+
+					SVerticalBox::Slot()
+					.AutoHeight()
+					.Padding(25, 0, 0, 0)
+					[
+						SNew(SCheckBox)
+							.IsChecked(this, &FModelDataEditorToolkit::GetCheckBoxState, AllCheckBoxIndex)
+							.OnCheckStateChanged(this, &FModelDataEditorToolkit::SetCheckBoxState, AllCheckBoxIndex)
+							.ToolTipText(FText::FromString(TEXT("Enables all current and future runtimes")))
+							[
+								SNew(STextBlock)
+									.Text(FText::FromString(TEXT("All")))
+									.Font(FAppStyle::GetFontStyle("PropertyWindow.NormalFont"))
+									.Margin(FMargin(4, 0))
+
+							]
+					]
+			];
+		for (int RuntimeIndex = 0; RuntimeIndex < Runtimes.Num(); RuntimeIndex++)
+		{
+			const FRuntimeArea& Runtime = Runtimes[RuntimeIndex];
+			TSharedPtr<SRuntimeExpandableArea> ExpandableArea;
+			VerticalBox->AddSlot()
+				.AutoHeight()
+				.Padding(4.0)
+				[
+					SAssignNew(ExpandableArea, SRuntimeExpandableArea)
+						.InitiallyCollapsed(true)
+						.BorderImage(FAppStyle::Get().GetBrush("Brushes.Header"))
+						.BodyBorderImage(FAppStyle::Get().GetBrush("Brushes.Recessed"))
+						.HeaderPadding(FMargin(2.f, 5.f))
+						.Padding(FMargin(15.f, 1.f, 0.f, 0.f))
+						.HeaderContent()
+						[
+							SNew(SHorizontalBox)
+							+
+								SHorizontalBox::Slot()
+								.AutoWidth()
+								.Padding(0, 0, 5, 0)
+								[
+									SNew(SCheckBox)
+										.IsChecked(this, &FModelDataEditorToolkit::GetCheckBoxState, RuntimeIndex)
+										.OnCheckStateChanged(this, &FModelDataEditorToolkit::SetCheckBoxState, RuntimeIndex)
+										.ToolTipText(FText::FromString(TEXT("Enables cooking and usage of this runtime.")))
+								]
+								+
+								SHorizontalBox::Slot()
+								.AutoWidth()
+								.Padding(0, 0, 10, 0)
+								.VAlign(VAlign_Center)
+								[
+									SNew(STextBlock)
+										.Text(FText::FromString(Runtime.Name))
+										.Font(FCoreStyle::Get().GetFontStyle("ExpandableArea.TitleFont"))
+								]
+								+
+								SHorizontalBox::Slot()
+								.VAlign(VAlign_Center)
+								[
+									SNew(STextBlock)
+										.Text(this, &FModelDataEditorToolkit::GetAreaErrorMessageText, RuntimeIndex)
+										.ColorAndOpacity(FSlateColor(EStyleColor::AccentRed))
+										.Font(FCoreStyle::Get().GetFontStyle("ExpandableArea.TitleFont"))
+								]
+						]
+						.BodyContent()
+						[
+							Runtime.DetailsView
+						]
+				];
+			ExpandableArea->IsExpandable.BindRaw(this, &FModelDataEditorToolkit::GetAreaIsExpandable, RuntimeIndex);
+		}
+		return Result;
+	}
+
+	void FModelDataEditorToolkit::ReloadSettings()
+	{
+		using EStatus = UNNEModelData::EGetRuntimeSettingsStatus;
+
+		ClearSettings();
+		for (FRuntimeArea& Runtime : Runtimes)
+		{
+			TArray<UObject*> AllSettings;
+			Runtime.IsVisible = true;
+			for (int32 ModelIndex = 0; ModelIndex < Models.Num(); ModelIndex++)
+			{
+				const UNNEModelData* Model = Models[ModelIndex];
+				UObject* Settings;
+				EStatus Status = Model->GetRuntimeSettings(Runtime.Name, Settings);
+				if (Status != EStatus::Success)
+				{
+					Runtime.IsVisible = false;
+					switch (Status)
+					{
+						case EStatus::RuntimeNotFound:
+						case EStatus::RuntimeUnsupported:
+							break;
+						case EStatus::InvalidCustomVersion:
+							Runtime.ErrorMessage = INVTEXT("Error: Invalid custom version. Please try to delete and reimport the model.");
+							break;
+						case EStatus::NewerCustomVersion:
+							Runtime.ErrorMessage = INVTEXT("Error: Asset was saved with a newer custom version than the current. Please open it with a newer engine version.");
+							break;
+						case EStatus::SerializationError:
+							Runtime.ErrorMessage = INVTEXT("Error: Serialization failed. Please try to delete and reimport the model.");
+							break;
+						default:
+							checkNoEntry();
+							break;
+					}
+					break;
+				}
+				Settings->AddToRoot();
+				AllSettings.Add(Settings);
+				const FSettingsIndex SettingsIndex(Runtime.Name, ModelIndex, Settings);
+				SettingsIDToIndex.Add(Settings->GetUniqueID(), SettingsIndex);
+			}
+			Runtime.DetailsView->SetObjects(AllSettings);
+		}
+	}
+
+	void FModelDataEditorToolkit::ClearSettings()
+	{
+		for (const auto& Element : SettingsIDToIndex)
+		{
+			Element.Value.Settings->RemoveFromRoot();
+		}
+		SettingsIDToIndex.Reset();
+	}
+
+	void FModelDataEditorToolkit::OnRuntimePropertyChanged(const FPropertyChangedEvent& Property) const
+	{
+		const FString TransactionDescription = TEXT("Edit ") + Property.GetPropertyName().ToString();
+		UKismetSystemLibrary::BeginTransaction(TEXT("NNE Model Data"), FText::FromString(*TransactionDescription), nullptr);
+		for (int Index = 0; Index < Property.GetNumObjectsBeingEdited(); Index++)
+		{
+			const uint32 SettingsID = Property.GetObjectBeingEdited(Index)->GetUniqueID();
+			check(SettingsIDToIndex.Contains(SettingsID));
+			const FSettingsIndex& SettingsIndex = SettingsIDToIndex[SettingsID];
+			UNNEModelData* Model = Models[SettingsIndex.ModelIndex];
+
+			UKismetSystemLibrary::TransactObject(Model);
+			Model->SetRuntimeSettings(SettingsIndex.RuntimeName, *SettingsIndex.Settings);
+		}
+		UKismetSystemLibrary::EndTransaction();
+	}
+
+	EVisibility FModelDataEditorToolkit::GetErrorMessageVisibility() const
+	{
+		return EmptySelectionModels->ModelIndices.IsEmpty() ? EVisibility::Collapsed : EVisibility::Visible;
+	}
+
+	FText FModelDataEditorToolkit::GetAreaErrorMessageText(const int RuntimeIndex) const
+	{
+		const FRuntimeArea& Runtime = Runtimes[RuntimeIndex];
+		if (!Runtime.IsVisible)
+		{
+			return Runtime.ErrorMessage;
+		}
+		return INVTEXT("");
+	}
+
+	bool FModelDataEditorToolkit::GetAreaIsExpandable(const int RuntimeIndex) const
+	{
+		return Runtimes[RuntimeIndex].IsVisible;
+	}
+
+	FString FModelDataEditorToolkit::ButtonIndexToRuntimeName(const int ButtonIndex) const
 	{
-		check(ModelData != nullptr);
-		return ModelData->GetTargetRuntimes();
+		if (ButtonIndex == AllCheckBoxIndex)
+		{
+			return TEXT("All");
+		}
+		check(ButtonIndex >= 0);
+		return Runtimes[ButtonIndex].Name;
 	}
 
-	void FModelDataEditorToolkit::SetTargetRuntimes(TArrayView<const FString> TargetRuntimes)
+	ECheckBoxState FModelDataEditorToolkit::GetCheckBoxState(const int ButtonIndex) const
 	{
-		check(ModelData != nullptr);
-		ModelData->Modify();
-		ModelData->SetTargetRuntimes(TargetRuntimes);
+		bool bAllChecked = true;
+		bool bAllUnchecked = true;
+		const FString RuntimeName = ButtonIndexToRuntimeName(ButtonIndex);
+		for(int ModelIndex = 0; ModelIndex < Models.Num(); ModelIndex++)
+		{
+			const UNNEModelData* Model = Models[ModelIndex];
+			const TArrayView<const FString> TargetRuntimes = Model->GetTargetRuntimes();
+			const bool bIsChecked = TargetRuntimes.IsEmpty() || (TargetRuntimes.Contains(RuntimeName) && !EmptySelectionModels->ModelIndices.Contains(ModelIndex));
+			bAllChecked &= bIsChecked;
+			bAllUnchecked &= !bIsChecked;
+		}
+		const ECheckBoxState State = bAllChecked   ? ECheckBoxState::Checked :
+							         bAllUnchecked ? ECheckBoxState::Unchecked :
+							         /*else*/        ECheckBoxState::Undetermined;
+		return State;
 	}
 
-} // UE::NNEEditor::Private
\ No newline at end of file
+	void FModelDataEditorToolkit::SetCheckBoxState(const ECheckBoxState, const int ButtonIndex)
+	{
+		// Determining the next check box state based on the current one.
+		// Normaly the rule is that if the state is in Undetermined the check box gets unchecked.
+		// In our case however this causes the All button to be unchecked (if it isn't allready). 
+		// Since this is unexpected and not an operation that can easily be reverted we favor
+		// checking the check box in this case.
+		const ECheckBoxState OldState = GetCheckBoxState(ButtonIndex);
+		const ECheckBoxState OldAllState = GetCheckBoxState(AllCheckBoxIndex);
+		const bool bFavorSelection = ButtonIndex != AllCheckBoxIndex && OldAllState != ECheckBoxState::Unchecked;
+		const bool bAddRuntime =  bFavorSelection ? OldState != ECheckBoxState::Checked : OldState == ECheckBoxState::Unchecked;
+
+		const FString RuntimeName = ButtonIndexToRuntimeName(ButtonIndex);
+		FString TransactionDescription = bAddRuntime ? TEXT("Add ") : TEXT("Remove ");
+		TransactionDescription += RuntimeName;
+		UKismetSystemLibrary::BeginTransaction(TEXT("NNE Model Data"), FText::FromString(*TransactionDescription), nullptr);
+		for (int ModelIndex = 0; ModelIndex < Models.Num(); ModelIndex++)
+		{
+			UNNEModelData* Model = Models[ModelIndex];
+			const TArrayView<const FString> CurrentRuntimes = Model->GetTargetRuntimes();
+			const bool bIsEmptySelection = EmptySelectionModels->ModelIndices.Contains(ModelIndex);
+			const bool bIsRuntimeSelected = CurrentRuntimes.IsEmpty() || CurrentRuntimes.Contains(RuntimeName);
+			const bool bIsButtonChecked = bIsRuntimeSelected && !bIsEmptySelection;
+			if (bAddRuntime == bIsButtonChecked)
+			{
+				continue;
+			}
+			if (bAddRuntime && bIsEmptySelection)
+			{
+				UKismetSystemLibrary::TransactObject(EmptySelectionModels);
+				EmptySelectionModels->ModelIndices.Remove(ModelIndex);
+			}
+			TArray<FString> NewRuntimes;
+			bool bUpdateRuntimes = bIsRuntimeSelected != bAddRuntime;
+			// All check box was pressed
+			if (ButtonIndex == AllCheckBoxIndex)
+			{
+				if (!bAddRuntime)
+				{
+					NewRuntimes = AllRuntimeNames;
+				}
+			}
+			else
+			{
+				if (bAddRuntime)
+				{
+					if (!bIsEmptySelection)
+					{
+						NewRuntimes = CurrentRuntimes;
+					}
+					NewRuntimes.Add(RuntimeName);
+				}
+				else
+				{
+					if (CurrentRuntimes.IsEmpty())
+					{
+						NewRuntimes = UE::NNE::GetAllRuntimeNames();
+						NewRuntimes.Remove(RuntimeName);
+					}
+					else
+					{
+						if (CurrentRuntimes.Num() > 1)
+						{
+							NewRuntimes = CurrentRuntimes;
+							NewRuntimes.Remove(RuntimeName);
+						}
+						else
+						{
+							bUpdateRuntimes = false;
+							UKismetSystemLibrary::TransactObject(EmptySelectionModels);
+							EmptySelectionModels->ModelIndices.Add(ModelIndex);
+						}
+					}
+				}
+			}
+			if (bUpdateRuntimes)
+			{
+				UKismetSystemLibrary::TransactObject(Model);
+				Model->SetTargetRuntimes(NewRuntimes);
+			}
+		}
+		UKismetSystemLibrary::EndTransaction();
+	}
+} // UE::NNEEditor::Private
diff --git a/Engine/Source/Editor/NNEEditor/Private/NNEEditorModelDataEditorToolkit.h b/Engine/Source/Editor/NNEEditor/Private/NNEEditorModelDataEditorToolkit.h
--- a/Engine/Source/Editor/NNEEditor/Private/NNEEditorModelDataEditorToolkit.h
+++ b/Engine/Source/Editor/NNEEditor/Private/NNEEditorModelDataEditorToolkit.h
@@ -4,31 +4,76 @@
 
 #include "CoreMinimal.h"
 #include "Toolkits/AssetEditorToolkit.h"
+#include "UObject/Object.h"
+
+#include "NNEEditorModelDataEditorToolkit.generated.h"
+
 
 class UNNEModelData;
 
+UCLASS(MinimalAPI)
+class UNNEEditorEmptySelectionModels : public UObject
+{
+	GENERATED_BODY()
+public:
+	UPROPERTY()
+	TArray<int> ModelIndices;
+};
+
 namespace UE::NNEEditor::Private
 {
 	class FModelDataEditorToolkit : public FAssetEditorToolkit
 	{
 	public:
-		void InitEditor(UNNEModelData* ModelData);
+		void InitEditor(const TArray<UNNEModelData*>& InModels);
 
 		// FAssetEditorToolkit interface
 		virtual void RegisterTabSpawners(const TSharedRef<class FTabManager>& TabManager) override;
 		virtual void UnregisterTabSpawners(const TSharedRef<class FTabManager>& TabManager) override;
 		virtual FName GetToolkitFName() const override { return FName("NNEModelDataEditor"); }
 		virtual FText GetBaseToolkitName() const override { return INVTEXT("NNE Model Data Editor"); }
+		virtual FText GetToolkitName() const override;
+		virtual FText GetToolkitToolTipText() const override;
 		virtual FString GetWorldCentricTabPrefix() const override { return TEXT("NNE Model Data Editor "); }
 		virtual FLinearColor GetWorldCentricTabColorScale() const override { return {}; }
+		virtual void OnClose() override;
 		// End of FAssetEditorToolkit interface
 
 	private:
-		TArrayView<const FString> GetTargetRuntimes() const;
-		void SetTargetRuntimes(TArrayView<const FString> TargetRuntimes);
 
-		TArray<FString> RegisteredRuntimes;
-		UNNEModelData* ModelData;
-	};
+		struct FSettingsIndex
+		{
+			const FString RuntimeName;
+			const int32 ModelIndex;
+			UObject* Settings;
+		};
 
-} // UE::NNEEditor::Private
\ No newline at end of file
+		struct FRuntimeArea
+		{
+			const FString Name;
+			const TSharedRef<IDetailsView> DetailsView;
+			bool IsVisible;
+			FText ErrorMessage;
+		};
+
+		static constexpr int AllCheckBoxIndex = -1;
+
+		TSharedRef<SDockTab> SpawnTab(const FSpawnTabArgs&);
+		void ReloadSettings();
+		void ClearSettings();
+		void OnRuntimePropertyChanged(const FPropertyChangedEvent& Property) const;
+		EVisibility GetErrorMessageVisibility() const;
+		FText GetAreaErrorMessageText(const int RuntimeIndex) const;
+		bool GetAreaIsExpandable(const int RuntimeIndex) const;
+		FString ButtonIndexToRuntimeName(const int ButtonIndex) const;
+		ECheckBoxState GetCheckBoxState(const int ButtonIndex) const;
+		void SetCheckBoxState(const ECheckBoxState NewState, const int Index);
+
+		TArray<UNNEModelData*> Models;
+		TArray<FString> AllRuntimeNames;
+		TArray<FRuntimeArea> Runtimes;
+		TMap<uint32, const FSettingsIndex> SettingsIDToIndex;
+		UNNEEditorEmptySelectionModels* EmptySelectionModels;
+		FDelegateHandle PostUndoRedoHandle;
+	};
+} // UE::NNEEditor::Private
diff --git a/Engine/Source/Editor/NNEEditor/Private/NNEEditorRuntimesWidget.cpp b/Engine/Source/Editor/NNEEditor/Private/NNEEditorRuntimesWidget.cpp
deleted file mode 100644
diff --git a/Engine/Source/Editor/NNEEditor/Private/NNEEditorRuntimesWidget.h b/Engine/Source/Editor/NNEEditor/Private/NNEEditorRuntimesWidget.h
deleted file mode 100644
diff --git a/Engine/Source/Runtime/NNE/Private/NNEModelData.cpp b/Engine/Source/Runtime/NNE/Private/NNEModelData.cpp
--- a/Engine/Source/Runtime/NNE/Private/NNEModelData.cpp
+++ b/Engine/Source/Runtime/NNE/Private/NNEModelData.cpp
@@ -3,9 +3,11 @@
 #include "NNEModelData.h"
 
 #include "NNE.h"
-#include "Serialization/CustomVersion.h"
 #include "UObject/AssetRegistryTagsContext.h"
 #include "UObject/WeakInterfacePtr.h"
+#include "Serialization/CustomVersion.h"
+#include "Serialization/MemoryReader.h"
+#include "Serialization/MemoryWriter.h"
 
 #include UE_INLINE_GENERATED_CPP_BY_NAME(NNEModelData)
 
@@ -18,6 +20,7 @@ namespace UE::NNE::ModelData
 		V2 = 2, // Re-arrange fields and store only ModelData in cooked assets
 		V3 = 3, // Adding AdditionalFileData
 		V4 = 4, // Support for > 2GB models
+		V5 = 5, // Adding RuntimeSettings
 		// New versions can be added above this line
 		VersionPlusOne,
 		Latest = VersionPlusOne - 1
@@ -48,37 +51,75 @@ namespace UE::NNE::ModelData
 		return RuntimesAsOneString;
 	}
 
-	inline TSharedPtr<UE::NNE::FSharedModelData> CreateModelData(const FString& RuntimeName, const FString& FileType, const TArray64<uint8>& FileData, const TMap<FString, TConstArrayView64<uint8>>& AdditionalFileData, const FGuid& FileId, const ITargetPlatform* TargetPlatform)
+	enum class ESerializationStatus : uint8
+	{
+		Success = 0,
+		InvalidCustomVersion = 1,
+		NewerCustomVersion = 2,
+		SerializationError = 3,
+	};
+
+	ESerializationStatus DeserializeRuntimeSettings(TConstArrayView64<uint8> InSerializedData, UObject* ModelSetting)
 	{
-		TWeakInterfacePtr<INNERuntime> NNERuntime = UE::NNE::GetRuntime<INNERuntime>(RuntimeName);
-		if (NNERuntime.IsValid())
+		FMemoryReaderView Reader(InSerializedData, /*bIsPersistent*/ true);
+		FCustomVersionContainer Versions;
+		TArray64<uint8> ObjectData;
+		Versions.Serialize(Reader, ECustomVersionSerializationFormat::Optimized);
+		Reader << ObjectData;
+		if (Reader.IsError() || !Reader.AtEnd())
 		{
-			INNERuntime::ECanCreateModelDataStatus CanCreateModelDataStatus = NNERuntime->CanCreateModelData(FileType, FileData, AdditionalFileData, FileId, TargetPlatform);
-			if (CanCreateModelDataStatus == INNERuntime::ECanCreateModelDataStatus::Ok)
-			{
-				return NNERuntime->CreateModelData(FileType, FileData, AdditionalFileData, FileId, TargetPlatform);
-			}
-			else if (CanCreateModelDataStatus == INNERuntime::ECanCreateModelDataStatus::FailFileIdNotSupported)
+			return ESerializationStatus::SerializationError;
+		}
+		TArray<FCustomVersionDifference> Diffs = FCurrentCustomVersions::Compare(Versions.GetAllVersions(), *ModelSetting->GetFName().ToString());
+		for (FCustomVersionDifference Diff : Diffs)
+		{
+			if (Diff.Type == ECustomVersionDifference::Invalid)
 			{
-				UE_LOG(LogNNE, Display, TEXT("Runtime %s does not support Filetype: %s, skipping the model data creation for model with id %s "), *RuntimeName, *FileType, *FileId.ToString(EGuidFormats::Digits).ToLower());
+				return ESerializationStatus::InvalidCustomVersion;
 			}
-			else
+			else if (Diff.Type == ECustomVersionDifference::Newer)
 			{
-				UE_LOG(LogNNE, Warning, TEXT("Runtime %s cannot create the model data with id %s (Filetype: %s)"), *RuntimeName, *FileId.ToString(EGuidFormats::Digits).ToLower(), *FileType);
+				return ESerializationStatus::NewerCustomVersion;
 			}
 		}
-		else
+		FMemoryReaderView ObjectReader(ObjectData, /*bIsPersistent*/ true);
+		ObjectReader.SetCustomVersions(Versions);
+		ModelSetting->Serialize(ObjectReader);
+		return ObjectReader.IsError() || !ObjectReader.AtEnd() ? ESerializationStatus::SerializationError : ESerializationStatus::Success;
+	}
+
+	ESerializationStatus SerializeRuntimeSettings(UObject* ModelSetting, TArray64<uint8>& OutSerializedData)
+	{
+		TArray64<uint8> ObjectData;
+		FMemoryWriter64 ObjectWriter(ObjectData, /*bIsPersistent*/ true);
+		ModelSetting->Serialize(ObjectWriter);
+		if (ObjectWriter.IsError())
+		{
+			return ESerializationStatus::SerializationError;
+		}
+		FCustomVersionContainer Versions = ObjectWriter.GetCustomVersions();
+		FMemoryWriter64 Writer(OutSerializedData, /*bIsPersistent*/ true);
+		Versions.Serialize(Writer, ECustomVersionSerializationFormat::Optimized);
+		Writer << ObjectData;
+		return Writer.IsError() ? ESerializationStatus::SerializationError : ESerializationStatus::Success;
+	}
+
+	template<class Value> 
+	const FString* FindLongestKeyWithPrefix(const TMap<FString, Value>& Map, const FString& Prefix)
+	{
+		const FString* Result = nullptr;
+		int32 LongestLength = 0;
+		for (const auto& Element : Map)
 		{
-			UE_LOG(LogNNE, Error, TEXT("No runtime '%s' found. Valid runtimes are: "), *RuntimeName);
-			TArray<FString> Runtimes = UE::NNE::GetAllRuntimeNames();
-			for (int32 i = 0; i < Runtimes.Num(); i++)
+			const FString& Key = Element.Key;
+			if (Key.Len() > LongestLength && Key.StartsWith(Prefix, ESearchCase::CaseSensitive))
 			{
-				UE_LOG(LogNNE, Error, TEXT("- %s"), *Runtimes[i]);
+				LongestLength = Key.Len();
+				Result = &Key;
 			}
 		}
-		return TSharedPtr<UE::NNE::FSharedModelData>();
+		return Result;
 	}
-
 } // UE::NNE::ModelData
 
 namespace UE::NNE
@@ -120,7 +161,7 @@ void UNNEModelData::Serialize(FArchive& Ar)
 		bool bWriteModelData = true;
 		if (Ar.IsCooking())
 		{
-			// Optimize storage: FileData and AdditionalFileData are not required anymore because we have the model and can cook it for every runtime
+			// Optimize storage: FileData, AdditionalFileData and RuntimeSettings are not required anymore because we have the model and can cook it for every runtime
 			TArray<FString> TmpTargetRuntimes;
 			Ar << TmpTargetRuntimes;
 			FString TmpFileType;
@@ -129,6 +170,8 @@ void UNNEModelData::Serialize(FArchive& Ar)
 			Ar << TmpFileData;
 			int32 NumAdditionalFileDataItems = 0;
 			Ar << NumAdditionalFileDataItems;
+			int32 NumRuntimeSettings = 0;
+			Ar << NumRuntimeSettings;
 			FGuid TmpGuid;
 			Ar << TmpGuid;
 
@@ -148,17 +191,9 @@ void UNNEModelData::Serialize(FArchive& Ar)
 					CookRuntimeNames.Append(GetTargetRuntimes());
 				}
 
-				TMap<FString, TConstArrayView64<uint8>> AdditionalFileDataView;
-				TSet<FString> Keys;
-				AdditionalFileData.GetKeys(Keys);
-				for (auto& Key : Keys)
-				{
-					AdditionalFileDataView.Emplace(Key, AdditionalFileData[Key]);
-				}
-
 				for (const FString& RuntimeName : CookRuntimeNames)
 				{
-					TSharedPtr<UE::NNE::FSharedModelData> SharedModelData = UE::NNE::ModelData::CreateModelData(RuntimeName, FileType, FileData, AdditionalFileDataView, FileId, Ar.GetArchiveState().CookingTarget());
+					TSharedPtr<UE::NNE::FSharedModelData> SharedModelData = CreateModelData(RuntimeName, Ar.GetArchiveState().CookingTarget());
 					if (SharedModelData.IsValid() && SharedModelData->GetView().Num() > 0)
 					{
 						ModelData.Add(RuntimeName, SharedModelData);
@@ -175,12 +210,18 @@ void UNNEModelData::Serialize(FArchive& Ar)
 
 			int32 NumAdditionalFileDataItems = AdditionalFileData.Num();
 			Ar << NumAdditionalFileDataItems;
-			TSet<FString> Keys;
-			AdditionalFileData.GetKeys(Keys);
-			for (auto& Key : Keys)
+			for (auto& Element : AdditionalFileData)
+			{
+				Ar << Element.Key;
+				Ar << Element.Value;
+			}
+
+			int32 NumRuntimeSettings = RuntimeSettings.Num();
+			Ar << NumRuntimeSettings;
+			for (auto& Element : RuntimeSettings)
 			{
-				Ar << Key;
-				Ar << AdditionalFileData[Key];
+				Ar << Element.Key;
+				Ar << Element.Value;
 			}
 
 			Ar << FileId;
@@ -222,6 +263,7 @@ void UNNEModelData::Serialize(FArchive& Ar)
 		TObjectPtr<class UAssetImportData> AssetImportData;
 		int32 NumItems = 0;
 		int32 NumAdditionalFileDataItems = 0;
+		int32 NumRuntimeSettings = 0;
 		FString Name;
 		uint32 MemoryAlignment = 0;
 		uint64 DataSize = 0;
@@ -230,7 +272,8 @@ void UNNEModelData::Serialize(FArchive& Ar)
 		void* RawData = nullptr;
 		int32 Index = 0;
 
-		switch (Ar.CustomVer(UE::NNE::ModelData::GUID))
+		const int32 Version = Ar.CustomVer(UE::NNE::ModelData::GUID);
+		switch (Version)
 		{
 		case UE::NNE::ModelData::Version::V0:
 			TargetRuntimes.Empty();
@@ -339,6 +382,39 @@ void UNNEModelData::Serialize(FArchive& Ar)
 			}
 			break;
 
+		case UE::NNE::ModelData::Version::V5:
+			Ar << TargetRuntimes;
+			Ar << FileType;
+			Ar << FileData;
+			Ar << NumAdditionalFileDataItems;
+			AdditionalFileData.Empty();
+			for (Index = 0; Index < NumAdditionalFileDataItems; Index++)
+			{
+				Ar << Name;
+				Ar << Data;
+				AdditionalFileData.Add(Name, Data);
+			}
+			Ar << NumRuntimeSettings;
+			RuntimeSettings.Reset();
+			for (Index = 0; Index < NumRuntimeSettings; Index++)
+			{
+				Ar << Name;
+				Ar << Data;
+				RuntimeSettings.Add(Name, Data);
+			}
+			Ar << FileId;
+			Ar << NumItems;
+			for (Index = 0; Index < NumItems; Index++)
+			{
+				Ar << Name;
+				Ar << MemoryAlignment;
+				Ar << DataSize;
+				RawData = FMemory::Malloc(DataSize, MemoryAlignment);
+				Ar.Serialize(RawData, DataSize);
+				ModelData.Add(Name, MakeShared<UE::NNE::FSharedModelData>(FSharedBuffer::TakeOwnership(RawData, DataSize, FMemory::Free), MemoryAlignment));
+			}
+			break;
+
 		default:
 			UE_LOG(LogNNE, Error, TEXT("Unknown asset version %d: Deserialisation failed, please reimport the original model."), Ar.CustomVer(UE::NNE::ModelData::GUID));
 			break;
@@ -354,8 +430,9 @@ void UNNEModelData::Init(const FString& Type, TConstArrayView64<uint8> Buffer, c
 	AdditionalFileData.Empty();
 	for (auto& Element : AdditionalBuffers)
 	{
-		AdditionalFileData.Emplace(Element.Key, AdditionalBuffers[Element.Key]);
+		AdditionalFileData.Emplace(Element.Key, Element.Value);
 	}
+	RuntimeSettings.Empty();
 	FPlatformMisc::CreateGuid(FileId);
 	ModelData.Empty();
 }
@@ -406,6 +483,103 @@ void UNNEModelData::ClearFileDataAndFileType()
 	AdditionalFileData.Empty();
 }
 
+UNNEModelData::EDeserializeRuntimeSettings UNNEModelData::DeserializeRuntimeSettings(const TMap<FString, TConstArrayView64<uint8>>& InAdditionalFileData, UObject& InOutSettings)
+{
+	using namespace UE::NNE::ModelData;
+	
+	if (const FString* Key = FindLongestKeyWithPrefix(InAdditionalFileData, RuntimeSettingsKeyPrefix))
+	{
+		ESerializationStatus Status = UE::NNE::ModelData::DeserializeRuntimeSettings(InAdditionalFileData[*Key], &InOutSettings);
+		switch (Status)
+		{
+			case ESerializationStatus::Success:
+				return EDeserializeRuntimeSettings::Success;
+			case ESerializationStatus::InvalidCustomVersion:
+				return EDeserializeRuntimeSettings::InvalidCustomVersion;
+			case ESerializationStatus::NewerCustomVersion:
+				return EDeserializeRuntimeSettings::NewerCustomVersion;
+			case ESerializationStatus::SerializationError:
+				return EDeserializeRuntimeSettings::SerializationError;
+			default:
+				checkNoEntry();
+				return EDeserializeRuntimeSettings::SerializationError;
+		}
+	}
+	return EDeserializeRuntimeSettings::Success;
+}
+
+UNNEModelData::EGetRuntimeSettingsStatus UNNEModelData::GetRuntimeSettings(const FString& InRuntimeName, UObject*& OutSettings) const
+{
+	check(!InRuntimeName.IsEmpty());
+
+	using namespace UE::NNE::ModelData;
+
+	TWeakInterfacePtr<INNERuntime> Runtime = UE::NNE::GetRuntime(InRuntimeName);
+	if (!Runtime.IsValid())
+	{
+		return EGetRuntimeSettingsStatus::RuntimeNotFound;
+	}
+	UObject* DefaultSetting = Runtime.Get()->CreateDefaultRuntimeSettings();
+	if (!DefaultSetting)
+	{
+		return EGetRuntimeSettingsStatus::RuntimeUnsupported;
+	}
+	if (const TArray64<uint8>* Data = RuntimeSettings.Find(InRuntimeName))
+	{
+		ESerializationStatus Status = UE::NNE::ModelData::DeserializeRuntimeSettings(*Data, DefaultSetting);
+		if (Status != ESerializationStatus::Success)
+		{
+			switch (Status)
+			{
+				case ESerializationStatus::InvalidCustomVersion:
+					return EGetRuntimeSettingsStatus::InvalidCustomVersion;
+				case ESerializationStatus::NewerCustomVersion:
+					return EGetRuntimeSettingsStatus::NewerCustomVersion;
+				case ESerializationStatus::SerializationError:
+					return EGetRuntimeSettingsStatus::SerializationError;
+				default:
+					checkNoEntry();
+					return EGetRuntimeSettingsStatus::SerializationError;
+			}
+		}
+	}
+	OutSettings = DefaultSetting;
+	return EGetRuntimeSettingsStatus::Success;
+}
+
+UNNEModelData::ESetRuntimeSettingsStatus UNNEModelData::SetRuntimeSettings(const FString& InRuntimeName, UObject& InSettings)
+{
+	check(!InRuntimeName.IsEmpty());
+
+	using namespace UE::NNE::ModelData;
+
+	TWeakInterfacePtr<INNERuntime> Runtime = UE::NNE::GetRuntime(InRuntimeName);
+	if (!Runtime.IsValid())
+	{
+		return ESetRuntimeSettingsStatus::RuntimeNotFound;
+	}
+	UObject* DefaultSetting = Runtime.Get()->CreateDefaultRuntimeSettings();
+	if (!DefaultSetting)
+	{
+		return ESetRuntimeSettingsStatus::RuntimeUnsupported;
+	}
+	if (DefaultSetting->GetClass() != InSettings.GetClass())
+	{
+		return ESetRuntimeSettingsStatus::WrongSettingType;
+	}
+	TArray64<uint8> SerializedSetting;
+	ESerializationStatus Status = SerializeRuntimeSettings(&InSettings, SerializedSetting);
+	if(Status != ESerializationStatus::Success)
+	{
+		check(Status == ESerializationStatus::SerializationError);
+		return ESetRuntimeSettingsStatus::SerializationError;
+	}
+	RuntimeSettings.Add(InRuntimeName, SerializedSetting);
+	// Invalidate model data cache
+	ModelData.Remove(InRuntimeName);
+	return ESetRuntimeSettingsStatus::Success;
+}
+
 FGuid UNNEModelData::GetFileId() const
 {
 	return FileId;
@@ -426,8 +600,7 @@ TSharedPtr<UE::NNE::FSharedModelData> UNNEModelData::GetModelData(const FString&
 	}
 
 	// Check if we have a local cache hit
-	TSharedPtr<UE::NNE::FSharedModelData>* LocalDataPtr = ModelData.Find(RuntimeName);
-	if (LocalDataPtr)
+	if (TSharedPtr<UE::NNE::FSharedModelData>* LocalDataPtr = ModelData.Find(RuntimeName))
 	{
 		return *LocalDataPtr;
 	}
@@ -439,16 +612,8 @@ TSharedPtr<UE::NNE::FSharedModelData> UNNEModelData::GetModelData(const FString&
 		return TSharedPtr<UE::NNE::FSharedModelData>();
 	}
 
-	TMap<FString, TConstArrayView64<uint8>> AdditionalFileDataView;
-	TSet<FString> Keys;
-	AdditionalFileData.GetKeys(Keys);
-	for (auto& Key : Keys)
-	{
-		AdditionalFileDataView.Emplace(Key, AdditionalFileData[Key]);
-	}
-
 	// Try to create the model
-	TSharedPtr<UE::NNE::FSharedModelData> CreatedData = UE::NNE::ModelData::CreateModelData(RuntimeName, FileType, FileData, AdditionalFileDataView, FileId, nullptr);
+	TSharedPtr<UE::NNE::FSharedModelData> CreatedData = CreateModelData(RuntimeName, nullptr);
 	if (!CreatedData.IsValid() || CreatedData->GetView().Num() < 1)
 	{
 		return TSharedPtr<UE::NNE::FSharedModelData>();
@@ -464,3 +629,54 @@ void UNNEModelData::ClearModelData()
 {
 	ModelData.Empty();
 }
+
+TSharedPtr<UE::NNE::FSharedModelData> UNNEModelData::CreateModelData(const FString& RuntimeName, const ITargetPlatform* TargetPlatform) const
+{
+	TWeakInterfacePtr<INNERuntime> NNERuntime = UE::NNE::GetRuntime<INNERuntime>(RuntimeName);
+	if (NNERuntime.IsValid())
+	{
+		TMap<FString, TConstArrayView64<uint8>> AdditionalFileDataView;
+		for (auto& Element : AdditionalFileData)
+		{
+			AdditionalFileDataView.Emplace(Element.Key, Element.Value);
+		}
+		if (const TArray64<uint8>* Settings = RuntimeSettings.Find(RuntimeName))
+		{
+			if (const FString* LongestKey = UE::NNE::ModelData::FindLongestKeyWithPrefix(AdditionalFileDataView, RuntimeSettingsKeyPrefix))
+			{
+				FString Key;
+				Key.Reserve(LongestKey->Len() + 1);
+				Key = *LongestKey;
+				Key.AppendChar(TEXT('_'));
+				AdditionalFileDataView.Add(Key, *Settings);
+			}
+			else
+			{
+				AdditionalFileDataView.Add(RuntimeSettingsKeyPrefix, *Settings);
+			}
+		}
+		INNERuntime::ECanCreateModelDataStatus CanCreateModelDataStatus = NNERuntime->CanCreateModelData(FileType, FileData, AdditionalFileDataView, FileId, TargetPlatform);
+		if (CanCreateModelDataStatus == INNERuntime::ECanCreateModelDataStatus::Ok)
+		{
+			return NNERuntime->CreateModelData(FileType, FileData, AdditionalFileDataView, FileId, TargetPlatform);
+		}
+		else if (CanCreateModelDataStatus == INNERuntime::ECanCreateModelDataStatus::FailFileIdNotSupported)
+		{
+			UE_LOG(LogNNE, Display, TEXT("Runtime %s does not support Filetype: %s, skipping the model data creation for model with id %s "), *RuntimeName, *FileType, *FileId.ToString(EGuidFormats::Digits).ToLower());
+		}
+		else
+		{
+			UE_LOG(LogNNE, Warning, TEXT("Runtime %s cannot create the model data with id %s (Filetype: %s)"), *RuntimeName, *FileId.ToString(EGuidFormats::Digits).ToLower(), *FileType);
+		}
+	}
+	else
+	{
+		UE_LOG(LogNNE, Error, TEXT("No runtime '%s' found. Valid runtimes are: "), *RuntimeName);
+		TArray<FString> Runtimes = UE::NNE::GetAllRuntimeNames();
+		for (int32 i = 0; i < Runtimes.Num(); i++)
+		{
+			UE_LOG(LogNNE, Error, TEXT("- %s"), *Runtimes[i]);
+		}
+	}
+	return TSharedPtr<UE::NNE::FSharedModelData>();
+}
diff --git a/Engine/Source/Runtime/NNE/Public/NNEModelData.h b/Engine/Source/Runtime/NNE/Public/NNEModelData.h
--- a/Engine/Source/Runtime/NNE/Public/NNEModelData.h
+++ b/Engine/Source/Runtime/NNE/Public/NNEModelData.h
@@ -76,6 +76,33 @@ class UNNEModelData : public UObject
 
 public:
 
+	enum class EDeserializeRuntimeSettings : uint8
+	{
+		Success = 0,
+		InvalidCustomVersion = 1,
+		NewerCustomVersion = 2,
+		SerializationError = 3,
+	};
+
+	enum class EGetRuntimeSettingsStatus : uint8
+	{
+		Success = 0,
+		RuntimeNotFound = 1,
+		RuntimeUnsupported = 2,
+		InvalidCustomVersion = 3,
+		NewerCustomVersion = 4,
+		SerializationError = 5,
+	};
+
+	enum class ESetRuntimeSettingsStatus : uint8
+	{
+		Success = 0,
+		RuntimeNotFound = 1,
+		RuntimeUnsupported = 2,
+		WrongSettingType = 3,
+		SerializationError = 4,
+	};
+
 	// UObject interface
 	NNE_API virtual void GetAssetRegistryTags(FAssetRegistryTagsContext Context) const override;
 	NNE_API virtual void Serialize(FArchive& Ar) override;
@@ -144,6 +171,38 @@ public:
 	 */
 	NNE_API void ClearFileDataAndFileType();
 
+
+	/**
+	 * Uses AdditionalFileData that is passed into CanCreateModelData and CreateModelData to populate InOutSettings
+	 *
+	 * @param InAdditionalFileData The AdditionalFileData argument that was passed into INNERuntime::CanCreateModelData or INNERuntime::CreateModelData
+	 * @param InOutSettings An instance of the RuntimeSettings that is returned from INNERuntime::CreateDefaultRuntimeSettings. If the user hasn't changed
+	 * the runtime settings this will be unchanged. Else InOutSettings will be deserialized. On failure the value of InOutSettings is undefined.
+	 *
+	 * @return Status of deserializing the settings.
+	 */
+	NNE_API static EDeserializeRuntimeSettings DeserializeRuntimeSettings(const TMap<FString, TConstArrayView64<uint8>>& InAdditionalFileData, UObject& InOutSettings);
+
+	/**
+	 * Gets the setting of the given runtime.
+	 *
+	 * @param InRuntimeName The name of the runtime whose setting to get.
+	 * @param OutSettings On success points to the created settings. On Failure OutSettings is left unchanged.
+	 *
+	 * @return Status of getting the setting.
+	 */
+	NNE_API EGetRuntimeSettingsStatus GetRuntimeSettings(const FString& InRuntimeName, UObject*& OutSettings) const;
+
+	/**
+	 * Sets the setting of the given runtime.
+	 *
+	 * @param InRuntimeName The name of the runtime whose setting to set.
+	 * @param InSettings The setting to set. This needs to have the same type as gets returned by the runtime through INNERuntime::CreateDefaultRuntimeSettings().
+	 *
+	 * @return Status of setting the setting.
+	 */
+	NNE_API ESetRuntimeSettingsStatus SetRuntimeSettings(const FString& InRuntimeName, UObject& InSettings);
+
 	/**
 	 * Get the FGuid identifying the FileData.
 	 *
@@ -192,6 +251,11 @@ private:
 	 */
 	TMap<FString, TArray64<uint8>> AdditionalFileData;
 
+	/**
+	 * Mapping between a runtime name and the serialized version of it's runtime settings.
+	 */
+	TMap<FString, TArray64<uint8>> RuntimeSettings;
+
 	/**
 	 * A Guid that uniquely identifies this model. This is used to cache optimized models in the editor.
 	 */
@@ -201,4 +265,14 @@ private:
 	 * The processed / optimized model data for the different runtimes.
 	 */
 	TMap<FString, TSharedPtr<UE::NNE::FSharedModelData>> ModelData;
+
+	/**
+	* Helper method for creating shared model data
+	*/
+	TSharedPtr<UE::NNE::FSharedModelData> CreateModelData(const FString& RuntimeName, const ITargetPlatform* TargetPlatform) const;
+
+	/**
+	* Prefix of the key that stores serialized runtime settings in the version of AdditionalFileData that is passed into INNERuntime::CanCreateModelData and INNERuntime::CreateModelData
+	*/
+	static inline const FString RuntimeSettingsKeyPrefix = TEXT("RuntimeSettings");
 };
diff --git a/Engine/Source/Runtime/NNE/Public/NNERuntime.h b/Engine/Source/Runtime/NNE/Public/NNERuntime.h
--- a/Engine/Source/Runtime/NNE/Public/NNERuntime.h
+++ b/Engine/Source/Runtime/NNE/Public/NNERuntime.h
@@ -54,10 +54,10 @@ public:
 	 *
 	 * @param FileType The type of file inside FileData. Corresponds to the file extension (e.g. 'onnx').
 	 * @param FileData The raw binary file of a neural network model.
-	 * @param AdditionalFileData Additional raw binary data of the neural network model.
+	 * @param AdditionalFileData Additional raw binary data of the neural network model. Can be passed to UNNEModelData::DeserializeRuntimeSettings.
 	 * @param FileId The unique identifier representing FileData.
 	 * @param TargetPlatform The Interface identifying the target platform for which the data needs to be created. A null pointer indicates the currently compiled/running platform.
-	 * @return True if the runtime is able to create model data, false otherwise.
+	 * @return Ok if the runtime is able to create model data, FailFileIdNotSupported if this fileId is not supported and Fail otherwise.
 	 */
 	virtual ECanCreateModelDataStatus CanCreateModelData(const FString& FileType, TConstArrayView64<uint8> FileData, const TMap<FString, TConstArrayView64<uint8>>& AdditionalFileData, const FGuid& FileId, const ITargetPlatform* TargetPlatform) const = 0;
 
@@ -66,7 +66,7 @@ public:
 	 *
 	 * @param FileType The type of file inside FileData. Corresponds to the file extension (e.g. 'onnx').
 	 * @param FileData The raw binary file of a neural network model.
-	 * @param AdditionalFileData Additional raw binary data of the neural network model.
+	 * @param AdditionalFileData Additional raw binary data of the neural network model. Can be passed to UNNEModelData::DeserializeRuntimeSettings.
 	 * @param FileId The unique identifier representing FileData.
 	 * @param TargetPlatform The Interface identifying the target platform for which the data needs to be created. A null pointer indicates the currently compiled/running platform.
 	 * @return A pointer to data representing the runtime specific representation of the model to be stored by UNNEModelData on success or an invalid pointer otherwise.
@@ -84,4 +84,14 @@ public:
 	 * @return A string that identifies the model data uniquely, e.g. by containing the FileId, model data version and target platform display name.
 	 */
 	virtual FString GetModelDataIdentifier(const FString& FileType, TConstArrayView64<uint8> FileData, const TMap<FString, TConstArrayView64<uint8>>& AdditionalFileData, const FGuid& FileId, const ITargetPlatform* TargetPlatform) const = 0;
+
+	/**
+	 * Create an instance of the default runtime settings. The Serialize(FArchive& Ar) and GetClass() methods are used to store and display these settings.
+	 * 
+	 * @return A new instance of the runtime settings or nullptr if runtime settings are not supported by this runtime.
+	 */
+	virtual UObject* CreateDefaultRuntimeSettings() const
+	{
+		return nullptr;
+	}
 };
\ No newline at end of file
