[NNE] Added scrollbars to runtime settings menu.

Details:
- Added SScrollBox around runtime settings content (without title)
- Minor title text format adjustment

#tests Editor, Win64, verified scroll bars present if content to larget for container

#rnx

#Synced-CL 51407547

diff --git a/Engine/Source/Editor/NNEEditor/Private/NNEEditorModelDataEditorToolkit.cpp b/Engine/Source/Editor/NNEEditor/Private/NNEEditorModelDataEditorToolkit.cpp
--- a/Engine/Source/Editor/NNEEditor/Private/NNEEditorModelDataEditorToolkit.cpp
+++ b/Engine/Source/Editor/NNEEditor/Private/NNEEditorModelDataEditorToolkit.cpp
@@ -8,9 +8,12 @@
 #include "NNEModelData.h"
 #include "NNERuntime.h"
 #include "PropertyEditorModule.h"
+#include "Styling/SlateTypes.h"
 #include "Widgets/Docking/SDockTab.h"
 #include "Widgets/Input/SCheckBox.h"
 #include "Widgets/Layout/SExpandableArea.h"
+#include "Widgets/Layout/SScrollBox.h"
+#include "Widgets/SBoxPanel.h"
 
 
 namespace UE::NNEEditor::Private
@@ -170,10 +173,10 @@ namespace UE::NNEEditor::Private
 
 	TSharedRef<SDockTab> FModelDataEditorToolkit::SpawnTab(const FSpawnTabArgs&)
 	{
-		TSharedPtr<SVerticalBox> VerticalBox;
+		TSharedPtr<SVerticalBox> ScrollContent;
 		TSharedRef<SDockTab> Result = SNew(SDockTab)
 			[
-				SAssignNew(VerticalBox, SVerticalBox)
+				SNew(SVerticalBox)
 					+
 					SVerticalBox::Slot()
 					.AutoHeight()
@@ -186,7 +189,8 @@ namespace UE::NNEEditor::Private
 							.Padding(0, 0, 15, 0)
 							[
 								SNew(STextBlock).Text(FText::FromString("Runtimes"))
-									.Font(FAppStyle::GetFontStyle("PropertyWindow.BoldFont"))
+									.TextStyle(FAppStyle::Get(), "DetailsView.CategoryTextStyle")
+									.TransformPolicy(ETextTransformPolicy::ToUpper)
 							]
 							+
 							SHorizontalBox::Slot()
@@ -201,27 +205,38 @@ namespace UE::NNEEditor::Private
 					]
 					+
 					SVerticalBox::Slot()
-					.AutoHeight()
-					.Padding(25, 0, 0, 0)
+					.FillHeight(1.0)
 					[
-						SNew(SCheckBox)
-							.IsChecked(this, &FModelDataEditorToolkit::GetCheckBoxState, AllCheckBoxIndex)
-							.OnCheckStateChanged(this, &FModelDataEditorToolkit::SetCheckBoxState, AllCheckBoxIndex)
-							.ToolTipText(FText::FromString(TEXT("Enables all current and future runtimes")))
+						SNew(SScrollBox)
+						+
+						SScrollBox::Slot()
+						[
+							SAssignNew(ScrollContent, SVerticalBox)
+							+
+							SVerticalBox::Slot()
+							.AutoHeight()
+							.Padding(25, 0, 0, 0)
 							[
-								SNew(STextBlock)
-									.Text(FText::FromString(TEXT("All")))
-									.Font(FAppStyle::GetFontStyle("PropertyWindow.NormalFont"))
-									.Margin(FMargin(4, 0))
-
+								SNew(SCheckBox)
+									.IsChecked(this, &FModelDataEditorToolkit::GetCheckBoxState, AllCheckBoxIndex)
+									.OnCheckStateChanged(this, &FModelDataEditorToolkit::SetCheckBoxState, AllCheckBoxIndex)
+									.ToolTipText(FText::FromString(TEXT("Enables all current and future runtimes")))
+									[
+										SNew(STextBlock)
+											.Text(FText::FromString(TEXT("All")))
+											.Font(FAppStyle::GetFontStyle("PropertyWindow.NormalFont"))
+											.Margin(FMargin(4, 0))
+
+									]
 							]
+						]
 					]
 			];
 		for (int RuntimeIndex = 0; RuntimeIndex < Runtimes.Num(); RuntimeIndex++)
 		{
 			const FRuntimeArea& Runtime = Runtimes[RuntimeIndex];
 			TSharedPtr<SRuntimeExpandableArea> ExpandableArea;
-			VerticalBox->AddSlot()
+			ScrollContent->AddSlot()
 				.AutoHeight()
 				.Padding(4.0)
 				[
